[Audio Insights] Added Audio::Trace::EventLog::SendEvent API to Blueprints

diff --git a/Engine/Plugins/AudioInsights/AudioInsights.uplugin b/Engine/Plugins/AudioInsights/AudioInsights.uplugin
--- a/Engine/Plugins/AudioInsights/AudioInsights.uplugin
+++ b/Engine/Plugins/AudioInsights/AudioInsights.uplugin
@@ -32,6 +32,10 @@
 		}
 	],
 	"Plugins": [
+		{
+			"Name": "AudioInsightsRuntime",
+			"Enabled": true
+		},
 		{
 			"Name": "AudioWidgets",
 			"Enabled": true
diff --git a/Engine/Plugins/AudioInsightsRuntime/AudioInsightsRuntime.uplugin b/Engine/Plugins/AudioInsightsRuntime/AudioInsightsRuntime.uplugin
new file mode 100644
--- /dev/null
+++ b/Engine/Plugins/AudioInsightsRuntime/AudioInsightsRuntime.uplugin
@@ -0,0 +1,25 @@
+{
+	"FileVersion": 3,
+	"Version": 1,
+	"VersionName": "1.0",
+	"FriendlyName": "Audio Insights Runtime",
+	"Description": "Enable sending profiling and debugging data via Blueprints to the Audio Insights plugin.",
+	"Category": "Audio",
+	"CreatedBy": "Epic Games, Inc.",
+	"CreatedByURL": "",
+	"DocsURL": "",
+	"MarketplaceURL": "",
+	"SupportURL": "",
+	"EnabledByDefault": true,
+	"CanContainContent": true,
+	"IsBetaVersion": true,
+	"IsExperimentalVersion": false,
+	"Installed": false,
+	"Modules": [
+		{
+			"Name": "AudioInsightsRuntime",
+			"Type": "Runtime",
+			"LoadingPhase": "Default"
+		}
+	]
+}
\ No newline at end of file
diff --git a/Engine/Plugins/AudioInsightsRuntime/Source/AudioInsightsRuntime/AudioInsightsRuntime.Build.cs b/Engine/Plugins/AudioInsightsRuntime/Source/AudioInsightsRuntime/AudioInsightsRuntime.Build.cs
new file mode 100644
--- /dev/null
+++ b/Engine/Plugins/AudioInsightsRuntime/Source/AudioInsightsRuntime/AudioInsightsRuntime.Build.cs
@@ -0,0 +1,30 @@
+// Copyright Epic Games, Inc. All Rights Reserved.
+
+using UnrealBuildTool;
+
+public class AudioInsightsRuntime : ModuleRules
+{
+	public AudioInsightsRuntime(ReadOnlyTargetRules Target) : base(Target)
+	{
+		PCHUsage = ModuleRules.PCHUsageMode.UseExplicitOrSharedPCHs;
+
+		PublicDependencyModuleNames.AddRange
+		(
+			new string[]
+			{
+				"Core",
+			}
+		);	
+		
+		PrivateDependencyModuleNames.AddRange
+		(
+			new string[]
+			{
+				"AudioMixerCore",
+				"CoreUObject",
+				"Engine",
+				"SlateCore",
+			}
+		);
+	}
+}
diff --git a/Engine/Plugins/AudioInsightsRuntime/Source/AudioInsightsRuntime/Private/AudioInsightsBlueprintLibrary.cpp b/Engine/Plugins/AudioInsightsRuntime/Source/AudioInsightsRuntime/Private/AudioInsightsBlueprintLibrary.cpp
new file mode 100644
--- /dev/null
+++ b/Engine/Plugins/AudioInsightsRuntime/Source/AudioInsightsRuntime/Private/AudioInsightsBlueprintLibrary.cpp
@@ -0,0 +1,87 @@
+// Copyright Epic Games, Inc. All Rights Reserved.
+
+#include "AudioInsightsBlueprintLibrary.h"
+
+#include "Audio/AudioTraceUtil.h"
+#include "AudioDevice.h"
+#include "AudioMixerTrace.h"
+#include "Engine/Engine.h"
+#include "Styling/SlateIconFinder.h"
+#include "UObject/Object.h"
+
+#include UE_INLINE_GENERATED_CPP_BY_NAME(AudioInsightsBlueprintLibrary)
+
+void UAudioInsightsBlueprintLibrary::LogAudioInsightsEvent(const UObject* WorldContextObject, const FString& EventName, const USoundBase* SoundAsset, const AActor* Actor)
+{
+#if UE_AUDIO_PROFILERTRACE_ENABLED
+	if (GEngine == nullptr || !GEngine->UseSound())
+	{
+		return;
+	}
+
+	const TObjectPtr<UWorld> ThisWorld = GEngine->GetWorldFromContextObject(WorldContextObject, EGetWorldErrorMode::LogAndReturnNull);
+	if (ThisWorld == nullptr || !ThisWorld->bAllowAudioPlayback)
+	{
+		return;
+	}
+
+	if (const FAudioDeviceHandle AudioDevice = ThisWorld->GetAudioDevice())
+	{
+		const TObjectPtr<UObject> SoundObjectPtr = SoundAsset != nullptr ? FSoftObjectPath(SoundAsset->GetPathName()).ResolveObject() : nullptr;
+		const FString AssetPathName = SoundObjectPtr != nullptr ? SoundObjectPtr->GetPathName() : FString();
+		const FString SoundBaseClassName = SoundObjectPtr != nullptr ? Audio::Trace::Util::GetSoundBaseAssetName(SoundObjectPtr->GetClass()) : FString();
+
+		const FString ActorLabel = Actor != nullptr ? Actor->GetActorNameOrLabel() : FString();
+		FString ActorIconName;
+#if WITH_EDITOR
+		if (Actor != nullptr)
+		{
+			FName IconName = Actor->GetCustomIconName();
+			if (IconName == NAME_None)
+			{
+				// Actor didn't specify an icon - fallback on the class icon
+				IconName = FSlateIconFinder::FindIconForClass(Actor->GetClass()).GetStyleName();
+			}
+
+			ActorIconName = IconName.ToString();
+		}
+#endif // WITH_EDITOR
+
+		Audio::Trace::EventLog::SendEvent(AudioDevice->DeviceID, EventName, INDEX_NONE, AssetPathName, ActorLabel, ActorIconName, SoundBaseClassName);
+	}
+#endif // UE_AUDIO_PROFILERTRACE_ENABLED
+}
+
+void UAudioInsightsBlueprintLibrary::LogAudioInsightsEventForAudioComponent(const UObject* WorldContextObject, const FString& EventName, const UAudioComponent* AudioComponent)
+{
+#if UE_AUDIO_PROFILERTRACE_ENABLED
+	if (GEngine == nullptr || !GEngine->UseSound() || AudioComponent == nullptr)
+	{
+		return;
+	}
+
+	const TObjectPtr<UWorld> ThisWorld = GEngine->GetWorldFromContextObject(WorldContextObject, EGetWorldErrorMode::LogAndReturnNull);
+	if (ThisWorld == nullptr || !ThisWorld->bAllowAudioPlayback)
+	{
+		return;
+	}
+
+	if (FAudioDeviceHandle AudioDevice = ThisWorld->GetAudioDevice())
+	{
+		const Audio::FDeviceId AudioDeviceID = AudioDevice->DeviceID;
+		AudioDevice->SendCommandToActiveSounds(AudioComponent->GetAudioComponentID(), [AudioDeviceID, EventName](FActiveSound& ActiveSound)
+			{
+				FAudioDeviceManager* AudioDeviceManager = GEngine->GetAudioDeviceManager();
+
+				if (AudioDeviceManager)
+				{
+					FAudioDeviceHandle AudioDevice = AudioDeviceManager->GetAudioDevice(AudioDeviceID);
+					if (AudioDevice.IsValid())
+					{
+						Audio::Trace::EventLog::SendActiveSoundEvent(ActiveSound, EventName);
+					}
+				}
+			});
+	}
+#endif // UE_AUDIO_PROFILERTRACE_ENABLED
+}
\ No newline at end of file
diff --git a/Engine/Plugins/AudioInsightsRuntime/Source/AudioInsightsRuntime/Private/AudioInsightsRuntimeModule.cpp b/Engine/Plugins/AudioInsightsRuntime/Source/AudioInsightsRuntime/Private/AudioInsightsRuntimeModule.cpp
new file mode 100644
--- /dev/null
+++ b/Engine/Plugins/AudioInsightsRuntime/Source/AudioInsightsRuntime/Private/AudioInsightsRuntimeModule.cpp
@@ -0,0 +1,18 @@
+// Copyright Epic Games, Inc. All Rights Reserved.
+#include "AudioInsightsRuntimeModule.h"
+
+#include "Modules/ModuleManager.h"
+
+namespace UE::Audio::Insights
+{
+	void FAudioInsightsRuntimeModule::StartupModule()
+	{
+	}
+
+	void FAudioInsightsRuntimeModule::ShutdownModule()
+	{
+	}
+
+} // namespace UE::Audio::Insights
+
+IMPLEMENT_MODULE(UE::Audio::Insights::FAudioInsightsRuntimeModule, AudioInsightsRuntime)
diff --git a/Engine/Plugins/AudioInsightsRuntime/Source/AudioInsightsRuntime/Private/AudioInsightsRuntimeModule.h b/Engine/Plugins/AudioInsightsRuntime/Source/AudioInsightsRuntime/Private/AudioInsightsRuntimeModule.h
new file mode 100644
--- /dev/null
+++ b/Engine/Plugins/AudioInsightsRuntime/Source/AudioInsightsRuntime/Private/AudioInsightsRuntimeModule.h
@@ -0,0 +1,22 @@
+// Copyright Epic Games, Inc. All Rights Reserved.
+#pragma once 
+
+#include "Modules/ModuleInterface.h"
+
+namespace UE::Audio::Insights
+{
+	class FAudioInsightsRuntimeModule final : public IModuleInterface
+	{
+	public:
+		FAudioInsightsRuntimeModule() = default;
+		virtual ~FAudioInsightsRuntimeModule() = default;
+
+		FAudioInsightsRuntimeModule(const FAudioInsightsRuntimeModule&) = delete;
+		FAudioInsightsRuntimeModule& operator=(const FAudioInsightsRuntimeModule&) = delete;
+		FAudioInsightsRuntimeModule(FAudioInsightsRuntimeModule&&) = delete;
+		FAudioInsightsRuntimeModule& operator=(FAudioInsightsRuntimeModule&&) = delete;
+
+		virtual void StartupModule() override;
+		virtual void ShutdownModule() override;
+	};
+} // namespace UE::Audio::Insights
diff --git a/Engine/Plugins/AudioInsightsRuntime/Source/AudioInsightsRuntime/Public/AudioInsightsBlueprintLibrary.h b/Engine/Plugins/AudioInsightsRuntime/Source/AudioInsightsRuntime/Public/AudioInsightsBlueprintLibrary.h
new file mode 100644
--- /dev/null
+++ b/Engine/Plugins/AudioInsightsRuntime/Source/AudioInsightsRuntime/Public/AudioInsightsBlueprintLibrary.h
@@ -0,0 +1,45 @@
+// Copyright Epic Games, Inc. All Rights Reserved.
+#pragma once
+
+#include "Kismet/BlueprintFunctionLibrary.h"
+
+#include "AudioInsightsBlueprintLibrary.generated.h"
+
+class AActor;
+class UAudioComponent;
+class UObject;
+class USoundBase;
+
+#define UE_API AUDIOINSIGHTSRUNTIME_API
+
+/** Static class with useful gameplay utility functions that can be called from both Blueprint and C++ */
+UCLASS(meta = (ScriptName = "AudioInsightsBlueprintLibrary"), MinimalAPI)
+class UAudioInsightsBlueprintLibrary : public UBlueprintFunctionLibrary
+{
+	GENERATED_BODY()
+
+public:
+	/**
+	 * Send an event message to the Audio Insights event log.
+	 * The filter category for the event can be added in the Audio Insights editor preferences
+	 * Requires the Audio Insights plug-in to be enabled.
+	 * @param EventName - the name that will appear in the event log for this event.
+	 * @param SoundAsset (optional) - the sound asset associated with this event.
+	 * @param Actor (optional) - the actor associated with this event.
+	 */
+	UFUNCTION(BlueprintCallable, Category = "Audio|AudioInsights", meta = (WorldContext = "WorldContextObject", AdvancedDisplay = "2"))
+	static UE_API void LogAudioInsightsEvent(const UObject* WorldContextObject, const FString& EventName, const USoundBase* SoundAsset = nullptr, const AActor* Actor = nullptr);
+
+	/**
+	 * Send an event message to the Audio Insights event log for a given AudioComponent.
+	 * The filter category for the event can be added in the Audio Insights editor preferences
+	 * Requires the Audio Insights plug-in to be enabled.
+	 * @param EventName - the name that will appear in the event log for this event
+	 * @param AudioComponent - the AudioComponent associated with this event
+	 */
+	UFUNCTION(BlueprintCallable, Category = "Audio|AudioInsights", meta = (WorldContext = "WorldContextObject"))
+	static UE_API void LogAudioInsightsEventForAudioComponent(const UObject* WorldContextObject, const FString& EventName, const class UAudioComponent* AudioComponent);
+
+};
+
+#undef UE_API
\ No newline at end of file
